- Fix CIS errors in OnlineSubsystemApple

diff --git a/Engine/Plugins/Online/OnlineSubsystemApple/Source/Private/OnlineIdentityInterfaceApple.cpp b/Engine/Plugins/Online/OnlineSubsystemApple/Source/Private/OnlineIdentityInterfaceApple.cpp
--- a/Engine/Plugins/Online/OnlineSubsystemApple/Source/Private/OnlineIdentityInterfaceApple.cpp
+++ b/Engine/Plugins/Online/OnlineSubsystemApple/Source/Private/OnlineIdentityInterfaceApple.cpp
@@ -326,7 +326,7 @@ FUniqueNetIdPtr FOnlineIdentityApple::CreateUniquePlayerId(uint8* Bytes, int32 S
 		if (StrLen > 0)
 		{
 			FString StrId((TCHAR*)Bytes);
-			return MakeShareable(new FUniqueNetIdApple(StrId));
+			return FUniqueNetIdApple::Create(StrId);
 		}
 	}
 
@@ -335,7 +335,7 @@ FUniqueNetIdPtr FOnlineIdentityApple::CreateUniquePlayerId(uint8* Bytes, int32 S
 
 FUniqueNetIdPtr FOnlineIdentityApple::CreateUniquePlayerId(const FString& Str)
 {
-	return MakeShareable(new FUniqueNetIdApple(Str));
+	return FUniqueNetIdApple::Create(Str);
 }
 
 FString FOnlineIdentityApple::GetPlayerNickname(int32 LocalUserNum) const
diff --git a/Engine/Plugins/Online/OnlineSubsystemApple/Source/Private/OnlineIdentityInterfaceApple.h b/Engine/Plugins/Online/OnlineSubsystemApple/Source/Private/OnlineIdentityInterfaceApple.h
--- a/Engine/Plugins/Online/OnlineSubsystemApple/Source/Private/OnlineIdentityInterfaceApple.h
+++ b/Engine/Plugins/Online/OnlineSubsystemApple/Source/Private/OnlineIdentityInterfaceApple.h
@@ -39,7 +39,7 @@ public:
 	// FUserOnlineAccountApple
 
 	explicit FUserOnlineAccountApple(const FString& InUserId = FString(), const FString& InAuthToken = FString())
-		: UserIdPtr(new FUniqueNetIdApple(InUserId))
+		: UserIdPtr(FUniqueNetIdApple::Create(InUserId))
 		, AuthToken(InAuthToken)
 	{
 	}
