[OnlineSubsystemApple]
* OnlineExternalUIInterfaceApple.h - Added missing include to set PACKAGE_SCOPE to correct protection level

diff --git a/Engine/Plugins/Online/OnlineSubsystemApple/Source/Private/OnlineExternalUIInterfaceApple.h b/Engine/Plugins/Online/OnlineSubsystemApple/Source/Private/OnlineExternalUIInterfaceApple.h
--- a/Engine/Plugins/Online/OnlineSubsystemApple/Source/Private/OnlineExternalUIInterfaceApple.h
+++ b/Engine/Plugins/Online/OnlineSubsystemApple/Source/Private/OnlineExternalUIInterfaceApple.h
@@ -3,6 +3,7 @@
 #pragma once
 
 #include "Interfaces/OnlineExternalUIInterface.h"
+#include "OnlineSubsystemApplePackage.h"
 
 class FOnlineSubsystemApple;
 
